Add tests to cover some missing cases of LatentIt and Async It

#rnx

diff --git a/Engine/Plugins/Tests/TestFramework/Source/TestFramework/Private/AutomationSpec.spec.cpp b/Engine/Plugins/Tests/TestFramework/Source/TestFramework/Private/AutomationSpec.spec.cpp
--- a/Engine/Plugins/Tests/TestFramework/Source/TestFramework/Private/AutomationSpec.spec.cpp
+++ b/Engine/Plugins/Tests/TestFramework/Source/TestFramework/Private/AutomationSpec.spec.cpp
@@ -138,4 +138,115 @@ void AutomationSpec::Define()
 			});
 		});
 	});
+
+	Describe("A spec async", [this]()
+	{
+		Describe("will run BeforeEach and AfterEach blocks", [this]()
+		{
+			BeforeEach([this]()
+			{
+				RunOrder = TEXT("A");
+			});
+
+			It("", EAsyncExecution::ThreadPool, [this]()
+			{
+				RunOrder += TEXT("B");
+			});
+
+			AfterEach([this]()
+			{
+				RunOrder += TEXT("C");
+				TestEqual("RunOrder", RunOrder, TEXT("ABC"));
+			});
+		});
+	});
+
+	Describe("LatentIt", [this]()
+	{
+		Describe("is called only once", [this]()
+		{
+			BeforeEach([this]()
+			{
+				RunOrder = TEXT("");
+			});
+
+			LatentIt("", [this](const FDoneDelegate Done)
+			{
+				RunOrder += TEXT("X");
+				if (RunOrder == TEXT("X"))
+				{
+					Async(EAsyncExecution::ThreadPool, [this, Done]() {
+						FPlatformProcess::Sleep(0.3f);
+						Done.Execute();
+					});
+				}
+			});
+
+			AfterEach([this]()
+			{
+				TestEqual("RunOrder", RunOrder, TEXT("X"));
+			});
+		});
+
+		Describe("can early exit", [this]()
+		{
+			BeforeEach([this]()
+			{
+				RunOrder = TEXT("");
+			});
+
+			LatentIt("", [this](const FDoneDelegate Done)
+			{
+				RunOrder += TEXT("X");
+				if (RunOrder == TEXT("X"))
+				{
+					Done.Execute();
+				}
+			});
+
+			AfterEach([this]()
+			{
+				TestEqual("RunOrder", RunOrder, TEXT("X"));
+			});
+		});
+
+		Describe("can be run async", [this]()
+		{
+			BeforeEach([this]()
+			{
+				RunOrder = TEXT("");
+			});
+
+			LatentIt("", EAsyncExecution::ThreadPool, [this](const FDoneDelegate Done)
+			{
+				FPlatformProcess::Sleep(0.1f);
+				RunOrder += TEXT("X");
+				Done.Execute();
+			});
+
+			AfterEach([this]()
+			{
+				TestEqual("RunOrder", RunOrder, TEXT("X"));
+			});
+		});
+
+		Describe("can exit early async", [this]()
+		{
+			BeforeEach([this]()
+			{
+				RunOrder = TEXT("");
+			});
+
+			LatentIt("", EAsyncExecution::ThreadPool, [this](const FDoneDelegate Done)
+			{
+				RunOrder += TEXT("X");
+				Done.Execute();
+			});
+
+			AfterEach([this]()
+			{
+				TestEqual("RunOrder", RunOrder, TEXT("X"));
+			});
+		});
+	});
 }
